Add TargetSettingsIniSectionName for Linux and LinuxArm64.

This was causing issues with not setting streaming audio to true for Compressed Audio

diff --git a/Engine/Config/Linux/DataDrivenPlatformInfo.ini b/Engine/Config/Linux/DataDrivenPlatformInfo.ini
--- a/Engine/Config/Linux/DataDrivenPlatformInfo.ini
+++ b/Engine/Config/Linux/DataDrivenPlatformInfo.ini
@@ -1,5 +1,6 @@
 [DataDrivenPlatformInfo]
 IniParent=Unix
+TargetSettingsIniSectionName=/Script/LinuxTargetPlatform.LinuxTargetSettings
 bHasDedicatedGamepad=false
 bDefaultInputStandardKeyboard=true
 bInputSupportConfigurable=true
diff --git a/Engine/Config/LinuxArm64/DataDrivenPlatformInfo.ini b/Engine/Config/LinuxArm64/DataDrivenPlatformInfo.ini
--- a/Engine/Config/LinuxArm64/DataDrivenPlatformInfo.ini
+++ b/Engine/Config/LinuxArm64/DataDrivenPlatformInfo.ini
@@ -1,5 +1,6 @@
 [DataDrivenPlatformInfo]
 IniParent=Unix
+TargetSettingsIniSectionName=/Script/LinuxTargetPlatform.LinuxTargetSettings
 bIsConfidential=false
 bHasDedicatedGamepad=false
 bDefaultInputStandardKeyboard=true
