Default bAgreeToUpload to false

If a game would like to implicitly upload set this to true

diff --git a/Engine/Config/Linux/LinuxEngine.ini b/Engine/Config/Linux/LinuxEngine.ini
--- a/Engine/Config/Linux/LinuxEngine.ini
+++ b/Engine/Config/Linux/LinuxEngine.ini
@@ -30,3 +30,6 @@ PlatformRequiresDataCrypto=True
 
 [OnlineSubsystem]
 LocalPlatformName=LNX
+
+[CrashReportClient]
+bAgreeToCrashUpload=false
diff --git a/Engine/Config/LinuxArm64/LinuxArm64Engine.ini b/Engine/Config/LinuxArm64/LinuxArm64Engine.ini
--- a/Engine/Config/LinuxArm64/LinuxArm64Engine.ini
+++ b/Engine/Config/LinuxArm64/LinuxArm64Engine.ini
@@ -27,4 +27,4 @@ g.TimeoutForBlockOnRenderFence=60000
 PlatformRequiresDataCrypto=True
 
 [CrashReportClient]
-bAgreeToCrashUpload=true
+bAgreeToCrashUpload=false
