Remove old comment about crash that is no longer happening

diff --git a/Engine/Config/Linux/LinuxEngine.ini b/Engine/Config/Linux/LinuxEngine.ini
--- a/Engine/Config/Linux/LinuxEngine.ini
+++ b/Engine/Config/Linux/LinuxEngine.ini
@@ -20,8 +20,6 @@ PoolSizeVRAMPercentage=70
 DeviceProfileSelectionModule="LinuxDeviceProfileSelector"
 
 [ConsoleVariables]
-; [RCL] 2017-10-09: if you get a crash with CacheOptimizeIndexBuffer, uncomment the line below and report a problem (refer to UE-23791).
-;r.TriangleOrderOptimization=2
 ; larger timeout since drivers may take longer time
 g.TimeoutForBlockOnRenderFence=60000
 
diff --git a/Engine/Config/LinuxArm64/LinuxArm64Engine.ini b/Engine/Config/LinuxArm64/LinuxArm64Engine.ini
--- a/Engine/Config/LinuxArm64/LinuxArm64Engine.ini
+++ b/Engine/Config/LinuxArm64/LinuxArm64Engine.ini
@@ -23,8 +23,6 @@ PoolSizeVRAMPercentage=70
 DeviceProfileSelectionModule="LinuxDeviceProfileSelector"
 
 [ConsoleVariables]
-; [RCL] 2017-10-09: if you get a crash with CacheOptimizeIndexBuffer, uncomment the line below and report a problem (refer to UE-23791).
-;r.TriangleOrderOptimization=2
 ; larger timeout since drivers may take longer time
 g.TimeoutForBlockOnRenderFence=60000
 
