Fixed linux arm 64 platform using the wrong cook platform name

diff --git a/Engine/Config/LinuxArm64/DataDrivenPlatformInfo.ini b/Engine/Config/LinuxArm64/DataDrivenPlatformInfo.ini
--- a/Engine/Config/LinuxArm64/DataDrivenPlatformInfo.ini
+++ b/Engine/Config/LinuxArm64/DataDrivenPlatformInfo.ini
@@ -12,7 +12,7 @@ bCanChangeGamepadType=true
 bSupportsTouch=false
 GlobalIdentifier=C0B45A339DE340C7BCE24C4715C3BABF
 
-OverridePlatformName=Linux
+OverrideCookPlatformName=Linux
 
 NormalIconPath=Launcher/Linux/Platform_Linux_24x
 LargeIconPath=Launcher/Linux/Platform_Linux_128x
