VulkanRHI: Fix LinuxArm64 packaging in SM6.

#rnx

diff --git a/Engine/Config/LinuxArm64/LinuxArm64Engine.ini b/Engine/Config/LinuxArm64/LinuxArm64Engine.ini
--- a/Engine/Config/LinuxArm64/LinuxArm64Engine.ini
+++ b/Engine/Config/LinuxArm64/LinuxArm64Engine.ini
@@ -28,3 +28,8 @@ g.TimeoutForBlockOnRenderFence=60000
 
 [PlatformCrypto]
 PlatformRequiresDataCrypto=True
+
+[SF_VULKAN_SM6]
+; Bindless is required for Vulkan SM6
+BindlessResources=Enabled
+BindlessSamplers=Enabled
\ No newline at end of file
